New experimental UnitTestFramework - for implementing Blueprint unit tests.

diff --git a/Engine/Plugins/Experimental/UnitTestFramework/Source/UnitTestFramework/Private/UTFGameAutomationTests.cpp b/Engine/Plugins/Experimental/UnitTestFramework/Source/UnitTestFramework/Private/UTFGameAutomationTests.cpp
new file mode 100644
--- /dev/null
+++ b/Engine/Plugins/Experimental/UnitTestFramework/Source/UnitTestFramework/Private/UTFGameAutomationTests.cpp
@@ -0,0 +1,147 @@
+// Copyright 1998-2015 Epic Games, Inc. All Rights Reserved.
+
+#include "UnitTestFrameworkPCH.h"
+#include "AutomationTest.h"
+#include "EngineGlobals.h"			// for GEngine
+#include "Engine/GameEngine.h"		// for GetGameWorld()
+#include "Engine/World.h"			// for PersistentLevel
+#include "UObjectHash.h"			// for StaticFindObjectFastExplicit
+#include "UTFUnitTestInterface.h"
+
+namespace UTFGameAutomationTestsImpl
+{
+	void GatherGameTestActors(TArray<FString>& OutBeautifiedNames, TArray<FString>& OutTestCommands)
+	{
+		auto GatherUnitTestActors = [&OutBeautifiedNames, &OutTestCommands](UWorld* GameWorld)
+		{
+			if (GameWorld != nullptr)
+			{
+				if (ULevel* GameLevel = GameWorld->PersistentLevel)
+				{
+					for (AActor* Actor : GameLevel->Actors)
+					{
+						if (Actor == nullptr)
+						{
+							continue;
+						}
+
+						UClass* ActorClass = Actor->GetClass();
+						if (ActorClass->ImplementsInterface(UUTFUnitTestInterface::StaticClass()))
+						{
+							const FString TestName = IUTFUnitTestInterface::Execute_GetTestName(Actor);
+							OutBeautifiedNames.Add(TestName);
+
+							OutTestCommands.Add(Actor->GetPathName());
+						}
+					}
+				}
+			}
+		};
+
+		if (UGameEngine* GameEngine = Cast<UGameEngine>(GEngine))
+		{
+			GatherUnitTestActors(GameEngine->GetGameWorld());
+		}
+		else
+		{
+			for (const FWorldContext& WorldContext : GEngine->GetWorldContexts())
+			{
+				if (WorldContext.WorldType == EWorldType::PIE ||
+					WorldContext.WorldType == EWorldType::Game)
+				{
+					GatherUnitTestActors(WorldContext.World());
+					break;
+				}
+			}
+		}
+	}
+}
+
+/*******************************************************************************
+ * FUTFResetGameActorTest
+ ******************************************************************************/
+
+IMPLEMENT_COMPLEX_AUTOMATION_TEST(FUTFResetGameActorTest, "UnitTestFramework.Game.ResetUnitTestActors", EAutomationTestFlags::ATF_Game | EAutomationTestFlags::ATF_Editor)
+
+//------------------------------------------------------------------------------
+void FUTFResetGameActorTest::GetTests(TArray<FString>& OutBeautifiedNames, TArray<FString>& OutTestCommands) const
+{
+	UTFGameAutomationTestsImpl::GatherGameTestActors(OutBeautifiedNames, OutTestCommands);
+}
+
+//------------------------------------------------------------------------------
+bool FUTFResetGameActorTest::RunTest(const FString& ObjectPath)
+{
+	AActor* TargetActor = FindObject<AActor>(/*Outer =*/nullptr, *ObjectPath);
+	if (TargetActor == nullptr)
+	{
+		AddError(FString::Printf(TEXT("Failed to find '%s' game object."), *ObjectPath));
+		return false;
+	}
+
+	UClass* ActorClass = TargetActor->GetClass();
+	if (!ActorClass->ImplementsInterface(UUTFUnitTestInterface::StaticClass()))
+	{
+		AddError(FString::Printf(TEXT("Target object does not implement the unit test interface: '%s'."), *ObjectPath));
+		return false;
+	}
+	
+	IUTFUnitTestInterface::Execute_ResetTest(TargetActor);
+	return true;
+}
+
+/*******************************************************************************
+ * FUTFGameActorTest
+ ******************************************************************************/
+
+IMPLEMENT_COMPLEX_AUTOMATION_TEST(FUTFGameActorTest, "UnitTestFramework.Game.UnitTestActors", EAutomationTestFlags::ATF_Game | EAutomationTestFlags::ATF_Editor)
+
+//------------------------------------------------------------------------------
+void FUTFGameActorTest::GetTests(TArray<FString>& OutBeautifiedNames, TArray<FString>& OutTestCommands) const
+{
+	UTFGameAutomationTestsImpl::GatherGameTestActors(OutBeautifiedNames, OutTestCommands);
+}
+
+//------------------------------------------------------------------------------
+bool FUTFGameActorTest::RunTest(const FString& ObjectPath)
+{
+	AActor* TargetActor = FindObject<AActor>(/*Outer =*/nullptr, *ObjectPath);
+	if (TargetActor == nullptr)
+	{
+		AddError(FString::Printf(TEXT("Failed to find '%s' game object."), *ObjectPath));
+		return false;
+	}
+
+	UClass* ActorClass = TargetActor->GetClass();
+	if (!ActorClass->ImplementsInterface(UUTFUnitTestInterface::StaticClass()))
+	{
+		AddError(FString::Printf(TEXT("Target object does not implement the unit test interface: '%s'."), *ObjectPath));
+		return false;
+	}
+
+	const FString TestName = IUTFUnitTestInterface::Execute_GetTestName(TargetActor);
+	
+	IUTFUnitTestInterface::Execute_ResetTest(TargetActor);
+	EUTFUnitTestResult TestResult = IUTFUnitTestInterface::Execute_GetTestResult(TargetActor);
+	
+	switch (TestResult)
+	{
+	case EUTFUnitTestResult::UTF_Success:
+		{
+			AddLogItem(FString::Printf(TEXT("Unit Test '%s': SUCCESS"), *TestName));
+		} break;
+
+	case EUTFUnitTestResult::UTF_Failure:
+		{
+			AddError(FString::Printf(TEXT("Unit Test '%s': FAILURE"), *TestName));
+		} break;
+
+	default:
+	case EUTFUnitTestResult::UTF_Unresolved:
+		{
+			AddWarning(FString::Printf(TEXT("Unit Test '%s': UNRESOLVED"), *TestName));
+		} break;
+	}
+
+	return (TestResult != EUTFUnitTestResult::UTF_Failure);
+}
diff --git a/Engine/Plugins/Experimental/UnitTestFramework/Source/UnitTestFramework/Private/UTFUnitTestInterface.cpp b/Engine/Plugins/Experimental/UnitTestFramework/Source/UnitTestFramework/Private/UTFUnitTestInterface.cpp
new file mode 100644
--- /dev/null
+++ b/Engine/Plugins/Experimental/UnitTestFramework/Source/UnitTestFramework/Private/UTFUnitTestInterface.cpp
@@ -0,0 +1,10 @@
+// Copyright 1998-2013 Epic Games, Inc. All Rights Reserved.
+
+#include "UnitTestFrameworkPCH.h"
+#include "UTFUnitTestInterface.h"
+
+//------------------------------------------------------------------------------
+UUTFUnitTestInterface::UUTFUnitTestInterface(const class FObjectInitializer& PCIP)
+	: Super(PCIP)
+{
+}
diff --git a/Engine/Plugins/Experimental/UnitTestFramework/Source/UnitTestFramework/Private/UnitTestFrameworkModule.cpp b/Engine/Plugins/Experimental/UnitTestFramework/Source/UnitTestFramework/Private/UnitTestFrameworkModule.cpp
new file mode 100644
--- /dev/null
+++ b/Engine/Plugins/Experimental/UnitTestFramework/Source/UnitTestFramework/Private/UnitTestFrameworkModule.cpp
@@ -0,0 +1,33 @@
+// Copyright 1998-2015 Epic Games, Inc. All Rights Reserved.
+
+#include "UnitTestFrameworkPCH.h"
+#include "UnitTestFrameworkModule.h"
+
+/** 
+ * 
+ */
+class FUnitTestFrameworkModule : public IUnitTestFrameworkModuleInterface
+{
+public:
+	virtual void StartupModule() override
+	{
+// 		FContentBrowserModule& ContentBrowserModule = FModuleManager::GetModuleChecked<FContentBrowserModule>(TEXT("ContentBrowser"));
+// 		TArray<FContentBrowserMenuExtender_SelectedAssets> ContentBrowserExtensionManager = ContentBrowserModule.GetAllAssetViewContextMenuExtenders();
+// 
+// 
+// 		FBlueprintEditorModule& BPEditor = FModuleManager::LoadModuleChecked<FBlueprintEditorModule>("Kismet");
+// 		MenuBarExtender = MakeShareable(new FExtender);
+// 
+// 		TSharedPtr< FExtender > MenuBarExtender;
+// 
+// 		MenuBarExtender->AddToolBarExtension("Asset", EExtensionHook::After, BPEditor.GetsSharedBlueprintEditorCommands(), FToolBarExtensionDelegate::CreateStatic(AddMyToolbarExtension));
+// 		BPEditor.GetMenuExtensibilityManager()->AddExtender(MenuBarExtender);
+// 
+// 		MenuBuilder.BeginSection("AssetContextAdvancedActions", LOCTEXT("AssetContextAdvancedActionsMenuHeading", "Advanced"));
+	}
+
+	virtual void ShutdownModule() override
+	{
+	}
+};
+IMPLEMENT_MODULE(FUnitTestFrameworkModule, UnitTestFramework);
diff --git a/Engine/Plugins/Experimental/UnitTestFramework/Source/UnitTestFramework/Private/UnitTestFrameworkPCH.h b/Engine/Plugins/Experimental/UnitTestFramework/Source/UnitTestFramework/Private/UnitTestFrameworkPCH.h
new file mode 100644
--- /dev/null
+++ b/Engine/Plugins/Experimental/UnitTestFramework/Source/UnitTestFramework/Private/UnitTestFrameworkPCH.h
@@ -0,0 +1,5 @@
+// Copyright 1998-2015 Epic Games, Inc. All Rights Reserved.
+
+#pragma once
+
+#include "Engine.h"
diff --git a/Engine/Plugins/Experimental/UnitTestFramework/Source/UnitTestFramework/Public/UTFUnitTestInterface.h b/Engine/Plugins/Experimental/UnitTestFramework/Source/UnitTestFramework/Public/UTFUnitTestInterface.h
new file mode 100644
--- /dev/null
+++ b/Engine/Plugins/Experimental/UnitTestFramework/Source/UnitTestFramework/Public/UTFUnitTestInterface.h
@@ -0,0 +1,41 @@
+// Copyright 1998-2013 Epic Games, Inc. All Rights Reserved.
+
+#pragma once
+
+#include "UTFUnitTestInterface.generated.h"
+
+/**  */
+UENUM(BlueprintType)
+enum class EUTFUnitTestResult
+{
+	UTF_Unresolved,
+	UTF_Success,
+	UTF_Failure,
+};
+
+/**  */
+UINTERFACE(BlueprintType, Category="Unit Test Framework")
+class UNITTESTFRAMEWORK_API UUTFUnitTestInterface : public UInterface
+{
+	GENERATED_UINTERFACE_BODY() 
+};
+
+/**  */
+class UNITTESTFRAMEWORK_API IUTFUnitTestInterface
+{
+	GENERATED_IINTERFACE_BODY()
+
+public:
+	UFUNCTION(BlueprintImplementableEvent, BlueprintCallable, Category = "Unit Test Framework", meta=(CallInEditor = "true"))
+	FString GetTestName();
+
+	UFUNCTION(BlueprintImplementableEvent, BlueprintCallable, Category = "Unit Test Framework", meta=(CallInEditor = "true"))
+	void ResetTest();
+
+	UFUNCTION(BlueprintImplementableEvent, BlueprintCallable, Category = "Unit Test Framework", meta=(CallInEditor = "true"))
+	EUTFUnitTestResult GetTestResult();
+
+protected:
+	UFUNCTION(BlueprintImplementableEvent, BlueprintCallable, Category = "Unit Test Framework", meta=(CallInEditor = "true"))
+	EUTFUnitTestResult RunUnitTest();
+};
diff --git a/Engine/Plugins/Experimental/UnitTestFramework/Source/UnitTestFramework/Public/UnitTestFrameworkModule.h b/Engine/Plugins/Experimental/UnitTestFramework/Source/UnitTestFramework/Public/UnitTestFrameworkModule.h
new file mode 100644
--- /dev/null
+++ b/Engine/Plugins/Experimental/UnitTestFramework/Source/UnitTestFramework/Public/UnitTestFrameworkModule.h
@@ -0,0 +1,10 @@
+// Copyright 1998-2015 Epic Games, Inc. All Rights Reserved.
+
+#pragma once
+
+/** 
+ * 
+ */
+class IUnitTestFrameworkModuleInterface : public IModuleInterface
+{
+};
diff --git a/Engine/Plugins/Experimental/UnitTestFramework/Source/UnitTestFramework/UnitTestFramework.Build.cs b/Engine/Plugins/Experimental/UnitTestFramework/Source/UnitTestFramework/UnitTestFramework.Build.cs
new file mode 100644
--- /dev/null
+++ b/Engine/Plugins/Experimental/UnitTestFramework/Source/UnitTestFramework/UnitTestFramework.Build.cs
@@ -0,0 +1,28 @@
+// Copyright 1998-2015 Epic Games, Inc. All Rights Reserved.
+
+using UnrealBuildTool;
+
+public class UnitTestFramework : ModuleRules
+{
+	public UnitTestFramework(TargetInfo Target)
+	{
+		PrivateIncludePaths.Add("UnitTestFramework/Private");
+
+		PublicDependencyModuleNames.AddRange(
+			new string[] {
+				"Core",
+				"CoreUObject",
+				"Engine",
+			}
+		);
+
+		PrivateDependencyModuleNames.AddRange(
+			new string[] {
+			}
+		);
+
+		if (UEBuildConfiguration.bBuildEditor == true)
+		{
+		}
+	}
+}
diff --git a/Engine/Plugins/Experimental/UnitTestFramework/Source/UnitTestFrameworkEditor/Private/UTFEditorAutomationTests.cpp b/Engine/Plugins/Experimental/UnitTestFramework/Source/UnitTestFrameworkEditor/Private/UTFEditorAutomationTests.cpp
new file mode 100644
--- /dev/null
+++ b/Engine/Plugins/Experimental/UnitTestFramework/Source/UnitTestFrameworkEditor/Private/UTFEditorAutomationTests.cpp
@@ -0,0 +1,125 @@
+// Copyright 1998-2015 Epic Games, Inc. All Rights Reserved.
+
+#include "UnitTestFrameworkEditorPCH.h"
+#include "AutomationTest.h"
+#include "AssetRegistryModule.h"	// for FAssetRegistry::GetAssets()
+#include "AssetData.h"
+#include "Engine/World.h"			// for UWorld::StaticClass()
+#include "PackageName.h"			// for FPackageName::LongPackageNameToFilename()
+#include "FileHelpers.h"			// for FEditorFileUtils::LoadMap()
+#include "Editor.h"					// for FEditorDelegates::OnMapOpened
+#include "UTFUnitTestInterface.h"
+#include "Script.h"					// for FEditorScriptExecutionGuard
+
+/*******************************************************************************
+ * FUTFEditorAutomationTests
+ ******************************************************************************/
+DEFINE_LOG_CATEGORY_STATIC(LogUTFEditorAutomationTests, Log, All);
+
+const FName FUTFEditorAutomationTests::UnitTestAssetTag(TEXT("UTF_UnitTestLevel"));
+
+/*******************************************************************************
+ * FUTFEditorMapTest
+ ******************************************************************************/
+
+IMPLEMENT_COMPLEX_AUTOMATION_TEST(FUTFEditorMapTest, "UnitTestFramework.Editor.RunUnitTestMaps", EAutomationTestFlags::ATF_Editor)
+
+//------------------------------------------------------------------------------
+void FUTFEditorMapTest::GetTests(TArray<FString>& OutBeautifiedNames, TArray<FString>& OutTestCommands) const
+{
+	FARFilter AssetFilter;
+	AssetFilter.ClassNames.Add(UWorld::StaticClass()->GetFName());
+	AssetFilter.TagsAndValues.Add(FUTFEditorAutomationTests::UnitTestAssetTag, TEXT("TRUE"));
+
+	TArray<FAssetData> UnitTestLevels;
+
+	IAssetRegistry& AssetRegistry = FModuleManager::LoadModuleChecked<FAssetRegistryModule>(TEXT("AssetRegistry")).Get();
+	AssetRegistry.GetAssets(AssetFilter, UnitTestLevels);
+
+	for (FAssetData const& LevelAsset : UnitTestLevels)
+	{
+		OutTestCommands.Add(LevelAsset.ObjectPath.ToString());
+		OutBeautifiedNames.Add(LevelAsset.AssetName.ToString());
+	}
+}
+
+//------------------------------------------------------------------------------
+bool FUTFEditorMapTest::RunTest(const FString& LevelAssetPath)
+{
+	// don't care about any map load warnings/errors, etc. 
+	SetSuppressLogs(true);
+
+	FDelegateHandle OnMapOpenedHandle = FEditorDelegates::OnMapOpened.AddLambda([this, &LevelAssetPath](const FString& /*Filename*/, bool /*bAsTemplate*/)
+	{ 
+		SetSuppressLogs(false);
+
+		FString PackageName, AssetName;
+		LevelAssetPath.Split(TEXT("."), &PackageName, &AssetName);
+
+		UPackage* WorldPackage = FindPackage(/*Outer =*/nullptr, *PackageName);
+		if (UWorld* TestWorld = FindObject<UWorld>(WorldPackage, *AssetName))
+		{
+			if (ULevel* TestLevel = TestWorld->PersistentLevel)
+			{
+				bool bExecutedTests = false;
+				for (AActor* Actor : TestLevel->Actors)
+				{
+					if (Actor == nullptr)
+					{
+						continue;
+					}
+
+					UClass* ActorClass = Actor->GetClass();
+					if (ActorClass->ImplementsInterface(UUTFUnitTestInterface::StaticClass()))
+					{
+						// notifies ProcessEvent() that it is ok to execute the following functions in the editor
+						FEditorScriptExecutionGuard ScriptGuard;
+
+						const FString TestName = IUTFUnitTestInterface::Execute_GetTestName(Actor);
+						// @TODO: consider the need for flagging tests as "runtime only" (cannot be ran in the editor)
+						EUTFUnitTestResult TestResult = IUTFUnitTestInterface::Execute_GetTestResult(Actor);
+						switch (TestResult)
+						{
+						case EUTFUnitTestResult::UTF_Success:
+							{
+								AddLogItem(FString::Printf(TEXT("Unit Test '%s': SUCCESS"), *TestName));
+							} break;
+
+						case EUTFUnitTestResult::UTF_Failure:
+							{
+								AddError(FString::Printf(TEXT("Unit Test '%s': FAILURE"), *TestName));
+							} break;
+
+						default:
+						case EUTFUnitTestResult::UTF_Unresolved:
+							{
+								AddWarning(FString::Printf(TEXT("Unit Test '%s': UNRESOLVED"), *TestName));
+							} break;
+						}
+						
+						bExecutedTests = true;
+					}
+				}
+
+				if (!bExecutedTests)
+				{
+					AddWarning(FString::Printf(TEXT("No editor unit tests found in '%s'."), *PackageName));
+				}
+			}
+			else
+			{
+				AddError(FString::Printf(TEXT("Null PersistentLevel after loading '%s'."), *PackageName));
+			}
+		}
+		else
+		{
+			AddError(FString::Printf(TEXT("Failed to find '%s' world object, post load."), *PackageName));
+		}
+	});
+
+	const FString MapFilePath = FPackageName::LongPackageNameToFilename(LevelAssetPath);
+	FEditorFileUtils::LoadMap(MapFilePath, /*bLoadAsTemplate =*/false, /*bShowProgress =*/false);
+	FEditorDelegates::OnMapOpened.Remove(OnMapOpenedHandle);
+
+	return (ExecutionInfo.Errors.Num() == 0);
+}
diff --git a/Engine/Plugins/Experimental/UnitTestFramework/Source/UnitTestFrameworkEditor/Private/UTFEditorAutomationTests.h b/Engine/Plugins/Experimental/UnitTestFramework/Source/UnitTestFrameworkEditor/Private/UTFEditorAutomationTests.h
new file mode 100644
--- /dev/null
+++ b/Engine/Plugins/Experimental/UnitTestFramework/Source/UnitTestFrameworkEditor/Private/UTFEditorAutomationTests.h
@@ -0,0 +1,11 @@
+// Copyright 1998-2015 Epic Games, Inc. All Rights Reserved.
+
+#pragma once
+
+/** 
+ * 
+ */
+struct FUTFEditorAutomationTests
+{
+	static const FName UnitTestAssetTag;
+};
diff --git a/Engine/Plugins/Experimental/UnitTestFramework/Source/UnitTestFrameworkEditor/Private/UnitTestFrameworkEditorModule.cpp b/Engine/Plugins/Experimental/UnitTestFramework/Source/UnitTestFrameworkEditor/Private/UnitTestFrameworkEditorModule.cpp
new file mode 100644
--- /dev/null
+++ b/Engine/Plugins/Experimental/UnitTestFramework/Source/UnitTestFrameworkEditor/Private/UnitTestFrameworkEditorModule.cpp
@@ -0,0 +1,65 @@
+// Copyright 1998-2015 Epic Games, Inc. All Rights Reserved.
+
+#include "UnitTestFrameworkEditorPCH.h"
+#include "UnitTestFrameworkEditorModule.h"
+#include "Engine/World.h"
+#include "UTFUnitTestInterface.h"
+#include "UTFEditorAutomationTests.h"
+
+/*******************************************************************************
+ * 
+ ******************************************************************************/
+
+namespace UnitTestFrameworkEditorModuleImpl
+{
+	/**  */
+	static void AddUnitTestingRegistryTags(const UWorld* World, TArray<UObject::FAssetRegistryTag>& TagsOut);
+}
+
+//------------------------------------------------------------------------------
+static void UnitTestFrameworkEditorModuleImpl::AddUnitTestingRegistryTags(const UWorld* World, TArray<UObject::FAssetRegistryTag>& TagsOut)
+{
+	if (ULevel* Level = World->PersistentLevel)
+	{
+		for (AActor* Actor : Level->Actors)
+		{
+			if (Actor == nullptr)
+			{
+				continue;
+			}
+
+			UClass* ActorClass = Actor->GetClass();
+			if (ActorClass->ImplementsInterface(UUTFUnitTestInterface::StaticClass()))
+			{
+				TagsOut.Add(UObject::FAssetRegistryTag(FUTFEditorAutomationTests::UnitTestAssetTag, TEXT("TRUE"), UObject::FAssetRegistryTag::TT_Hidden));
+				break;
+			}
+		}
+	}
+}
+
+/*******************************************************************************
+ * FUnitTestFrameworkEditorModule
+ ******************************************************************************/
+
+/** 
+ * 
+ */
+class FUnitTestFrameworkEditorModule : public IUnitTestFrameworkEditorModuleInterface
+{
+public:
+	virtual void StartupModule() override
+	{
+		GetWorldAssetTagsHandle = FWorldDelegates::GetAssetTags.AddStatic(UnitTestFrameworkEditorModuleImpl::AddUnitTestingRegistryTags);
+	}
+
+	virtual void ShutdownModule() override
+	{
+		FWorldDelegates::GetAssetTags.Remove(GetWorldAssetTagsHandle);
+		GetWorldAssetTagsHandle.Reset();
+	}
+
+private: 
+	FDelegateHandle GetWorldAssetTagsHandle;
+};
+IMPLEMENT_MODULE(FUnitTestFrameworkEditorModule, UnitTestFrameworkEditor);
diff --git a/Engine/Plugins/Experimental/UnitTestFramework/Source/UnitTestFrameworkEditor/Private/UnitTestFrameworkEditorPCH.h b/Engine/Plugins/Experimental/UnitTestFramework/Source/UnitTestFrameworkEditor/Private/UnitTestFrameworkEditorPCH.h
new file mode 100644
--- /dev/null
+++ b/Engine/Plugins/Experimental/UnitTestFramework/Source/UnitTestFrameworkEditor/Private/UnitTestFrameworkEditorPCH.h
@@ -0,0 +1,5 @@
+// Copyright 1998-2015 Epic Games, Inc. All Rights Reserved.
+
+#pragma once
+
+#include "Engine.h"
diff --git a/Engine/Plugins/Experimental/UnitTestFramework/Source/UnitTestFrameworkEditor/Public/UnitTestFrameworkEditorModule.h b/Engine/Plugins/Experimental/UnitTestFramework/Source/UnitTestFrameworkEditor/Public/UnitTestFrameworkEditorModule.h
new file mode 100644
--- /dev/null
+++ b/Engine/Plugins/Experimental/UnitTestFramework/Source/UnitTestFrameworkEditor/Public/UnitTestFrameworkEditorModule.h
@@ -0,0 +1,10 @@
+// Copyright 1998-2015 Epic Games, Inc. All Rights Reserved.
+
+#pragma once
+
+/** 
+ * 
+ */
+class IUnitTestFrameworkEditorModuleInterface : public IModuleInterface
+{
+};
diff --git a/Engine/Plugins/Experimental/UnitTestFramework/Source/UnitTestFrameworkEditor/UnitTestFrameworkEditor.Build.cs b/Engine/Plugins/Experimental/UnitTestFramework/Source/UnitTestFrameworkEditor/UnitTestFrameworkEditor.Build.cs
new file mode 100644
--- /dev/null
+++ b/Engine/Plugins/Experimental/UnitTestFramework/Source/UnitTestFrameworkEditor/UnitTestFrameworkEditor.Build.cs
@@ -0,0 +1,28 @@
+// Copyright 1998-2015 Epic Games, Inc. All Rights Reserved.
+
+using UnrealBuildTool;
+
+public class UnitTestFrameworkEditor : ModuleRules
+{
+	public UnitTestFrameworkEditor(TargetInfo Target)
+	{
+        PrivateIncludePaths.Add("UnitTestFrameworkEditor/Private");
+
+		PublicDependencyModuleNames.AddRange(
+			new string[] {
+				"Core",
+				"CoreUObject",
+				"Engine"
+			}
+		);
+
+		PrivateDependencyModuleNames.AddRange(
+			new string[] {
+                "UnitTestFramework",
+                "UnrealEd",
+                "AssetRegistry",
+			}
+		);
+	}
+}
+ 
\ No newline at end of file
diff --git a/Engine/Plugins/Experimental/UnitTestFramework/UnitTestFramework.uplugin b/Engine/Plugins/Experimental/UnitTestFramework/UnitTestFramework.uplugin
new file mode 100644
--- /dev/null
+++ b/Engine/Plugins/Experimental/UnitTestFramework/UnitTestFramework.uplugin
@@ -0,0 +1,31 @@
+{
+	"FileVersion" : 3,
+	
+	"FriendlyName"     : "Unit Test Framework",
+	"Version"          : 1,
+	"VersionName"      : "1.0",
+	"CreatedBy"        : "Epic Games, Inc.",
+	"CreatedByURL"     : "http://epicgames.com",
+	"EngineVersion"    : "4.9",
+	"Description"      : "",
+	"Category"         : "Experimental",
+	"EnabledByDefault" : false,
+
+	"Modules" :
+	[
+		{
+			"Name"               : "UnitTestFramework",
+			"Type"               : "Runtime",
+			"LoadingPhase"       : "Default",
+			"BlacklistPlatforms" : []
+		},
+		{
+			"Name"               : "UnitTestFrameworkEditor",
+			"Type"               : "Editor",
+			"LoadingPhase"       : "Default",
+			"BlacklistPlatforms" : []
+		},
+	],
+
+	"CanContainContent" : true
+}
\ No newline at end of file
