Fix to CIS: missing include file.

diff --git a/Engine/Plugins/Experimental/UnitTestFramework/Source/UnitTestFrameworkEditor/Private/UTFEditorAutomationTests.cpp b/Engine/Plugins/Experimental/UnitTestFramework/Source/UnitTestFrameworkEditor/Private/UTFEditorAutomationTests.cpp
--- a/Engine/Plugins/Experimental/UnitTestFramework/Source/UnitTestFrameworkEditor/Private/UTFEditorAutomationTests.cpp
+++ b/Engine/Plugins/Experimental/UnitTestFramework/Source/UnitTestFrameworkEditor/Private/UTFEditorAutomationTests.cpp
@@ -1,6 +1,7 @@
 // Copyright 1998-2015 Epic Games, Inc. All Rights Reserved.
 
 #include "UnitTestFrameworkEditorPCH.h"
+#include "UTFEditorAutomationTests.h"
 #include "AutomationTest.h"
 #include "AssetRegistryModule.h"	// for FAssetRegistry::GetAssets()
 #include "AssetData.h"
