Amending the new "unit test framework" so that you have to manually flag actor tests that are acceptable to run in the editor (as opposed to in-game, at runtime).

diff --git a/Engine/Plugins/Experimental/UnitTestFramework/Source/UnitTestFrameworkEditor/Private/UTFEditorAutomationTests.cpp b/Engine/Plugins/Experimental/UnitTestFramework/Source/UnitTestFrameworkEditor/Private/UTFEditorAutomationTests.cpp
--- a/Engine/Plugins/Experimental/UnitTestFramework/Source/UnitTestFrameworkEditor/Private/UTFEditorAutomationTests.cpp
+++ b/Engine/Plugins/Experimental/UnitTestFramework/Source/UnitTestFrameworkEditor/Private/UTFEditorAutomationTests.cpp
@@ -11,13 +11,31 @@
 #include "Editor.h"					// for FEditorDelegates::OnMapOpened
 #include "UTFUnitTestInterface.h"
 #include "Script.h"					// for FEditorScriptExecutionGuard
+#include "Engine/Blueprint.h"
 
 /*******************************************************************************
  * FUTFEditorAutomationTests
  ******************************************************************************/
 DEFINE_LOG_CATEGORY_STATIC(LogUTFEditorAutomationTests, Log, All);
 
-const FName FUTFEditorAutomationTests::UnitTestAssetTag(TEXT("UTF_UnitTestLevel"));
+const FName FUTFEditorAutomationTests::InEditorTestTag(TEXT("UTF_InEditorTest"));
+const FName FUTFEditorAutomationTests::UnitTestLevelTag(TEXT("UTF_UnitTestLevel"));
+
+//------------------------------------------------------------------------------
+bool FUTFEditorAutomationTests::IsClassSuitableForInEditorTesting(UClass* Class)
+{
+	if (Class->ImplementsInterface(UUTFUnitTestInterface::StaticClass()))
+	{
+		if (UBlueprint* Blueprint = CastChecked<UBlueprint>(Class->ClassGeneratedBy))
+		{
+			UPackage* BlueprintPackage = Blueprint->GetOutermost();
+			UMetaData* MetaData = BlueprintPackage->GetMetaData();
+
+			return MetaData->HasValue(Blueprint, FUTFEditorAutomationTests::InEditorTestTag);
+		}
+	}
+	return false;
+}
 
 /*******************************************************************************
  * FUTFEditorMapTest
@@ -30,7 +48,7 @@ void FUTFEditorMapTest::GetTests(TArray<FString>& OutBeautifiedNames, TArray<FSt
 {
 	FARFilter AssetFilter;
 	AssetFilter.ClassNames.Add(UWorld::StaticClass()->GetFName());
-	AssetFilter.TagsAndValues.Add(FUTFEditorAutomationTests::UnitTestAssetTag, TEXT("TRUE"));
+	AssetFilter.TagsAndValues.Add(FUTFEditorAutomationTests::UnitTestLevelTag, TEXT("TRUE"));
 
 	TArray<FAssetData> UnitTestLevels;
 
@@ -73,6 +91,12 @@ bool FUTFEditorMapTest::RunTest(const FString& LevelAssetPath)
 					UClass* ActorClass = Actor->GetClass();
 					if (ActorClass->ImplementsInterface(UUTFUnitTestInterface::StaticClass()))
 					{
+						if (!FUTFEditorAutomationTests::IsClassSuitableForInEditorTesting(ActorClass))
+						{
+							AddLogItem(FString::Printf(TEXT("Skipping test for '%s', as it is not marked safe for in-editor testing."), *ActorClass->GetName()));
+							continue;
+						}
+
 						// notifies ProcessEvent() that it is ok to execute the following functions in the editor
 						FEditorScriptExecutionGuard ScriptGuard;
 
@@ -124,3 +148,87 @@ bool FUTFEditorMapTest::RunTest(const FString& LevelAssetPath)
 
 	return (ExecutionInfo.Errors.Num() == 0);
 }
+
+/*******************************************************************************
+* FUTFEditorActorTest
+******************************************************************************/
+
+IMPLEMENT_COMPLEX_AUTOMATION_TEST(FUTFEditorActorTest, "UnitTestFramework.Editor.UnitTestActors", EAutomationTestFlags::ATF_Editor)
+
+//------------------------------------------------------------------------------
+void FUTFEditorActorTest::GetTests(TArray<FString>& OutBeautifiedNames, TArray<FString>& OutTestCommands) const
+{
+	for (const FWorldContext& WorldContext : GEngine->GetWorldContexts())
+	{
+		if (WorldContext.WorldType == EWorldType::Editor)
+		{
+			if (UWorld* EditorWorld = WorldContext.World())
+			{
+				if (ULevel* EditorLevel = EditorWorld->PersistentLevel)
+				{
+					for (AActor* Actor : EditorLevel->Actors)
+					{
+						if (Actor == nullptr)
+						{
+							continue;
+						}
+
+						UClass* ActorClass = Actor->GetClass();
+						if (FUTFEditorAutomationTests::IsClassSuitableForInEditorTesting(ActorClass))
+						{
+							const FString TestName = IUTFUnitTestInterface::Execute_GetTestName(Actor);
+							OutBeautifiedNames.Add(TestName);
+
+							OutTestCommands.Add(Actor->GetPathName());
+						}
+					}
+				}
+			}
+			break;
+		}
+	}
+}
+
+//------------------------------------------------------------------------------
+bool FUTFEditorActorTest::RunTest(const FString& ObjectPath)
+{
+	AActor* TargetActor = FindObject<AActor>(/*Outer =*/nullptr, *ObjectPath);
+	if (TargetActor == nullptr)
+	{
+		AddError(FString::Printf(TEXT("Failed to find '%s' game object."), *ObjectPath));
+		return false;
+	}
+
+	UClass* ActorClass = TargetActor->GetClass();
+	if (!ActorClass->ImplementsInterface(UUTFUnitTestInterface::StaticClass()))
+	{
+		AddError(FString::Printf(TEXT("Target object does not implement the unit test interface: '%s'."), *ObjectPath));
+		return false;
+	}
+
+	const FString TestName = IUTFUnitTestInterface::Execute_GetTestName(TargetActor);
+	
+	IUTFUnitTestInterface::Execute_ResetTest(TargetActor);
+	EUTFUnitTestResult TestResult = IUTFUnitTestInterface::Execute_GetTestResult(TargetActor);
+	
+	switch (TestResult)
+	{
+	case EUTFUnitTestResult::UTF_Success:
+		{
+			AddLogItem(FString::Printf(TEXT("Unit Test '%s': SUCCESS"), *TestName));
+		} break;
+
+	case EUTFUnitTestResult::UTF_Failure:
+		{
+			AddError(FString::Printf(TEXT("Unit Test '%s': FAILURE"), *TestName));
+		} break;
+
+	default:
+	case EUTFUnitTestResult::UTF_Unresolved:
+		{
+			AddWarning(FString::Printf(TEXT("Unit Test '%s': UNRESOLVED"), *TestName));
+		} break;
+	}
+
+	return (TestResult != EUTFUnitTestResult::UTF_Failure);
+}
diff --git a/Engine/Plugins/Experimental/UnitTestFramework/Source/UnitTestFrameworkEditor/Private/UTFEditorAutomationTests.h b/Engine/Plugins/Experimental/UnitTestFramework/Source/UnitTestFrameworkEditor/Private/UTFEditorAutomationTests.h
--- a/Engine/Plugins/Experimental/UnitTestFramework/Source/UnitTestFrameworkEditor/Private/UTFEditorAutomationTests.h
+++ b/Engine/Plugins/Experimental/UnitTestFramework/Source/UnitTestFrameworkEditor/Private/UTFEditorAutomationTests.h
@@ -7,5 +7,9 @@
  */
 struct FUTFEditorAutomationTests
 {
-	static const FName UnitTestAssetTag;
+	static const FName InEditorTestTag;
+	static const FName UnitTestLevelTag;
+
+	/**  */
+	static bool IsClassSuitableForInEditorTesting(UClass* Class);
 };
diff --git a/Engine/Plugins/Experimental/UnitTestFramework/Source/UnitTestFrameworkEditor/Private/UnitTestFrameworkEditorModule.cpp b/Engine/Plugins/Experimental/UnitTestFramework/Source/UnitTestFrameworkEditor/Private/UnitTestFrameworkEditorModule.cpp
--- a/Engine/Plugins/Experimental/UnitTestFramework/Source/UnitTestFrameworkEditor/Private/UnitTestFrameworkEditorModule.cpp
+++ b/Engine/Plugins/Experimental/UnitTestFramework/Source/UnitTestFrameworkEditor/Private/UnitTestFrameworkEditorModule.cpp
@@ -5,6 +5,13 @@
 #include "Engine/World.h"
 #include "UTFUnitTestInterface.h"
 #include "UTFEditorAutomationTests.h"
+#include "BlueprintEditorModule.h"
+#include "UTFEditorStyle.h"
+#include "MultiBoxBuilder.h" // for FExtender/FToolBarBuilder
+#include "SlateIcon.h"
+#include "BlueprintEditor.h" // for RegenerateMenusAndToolbars
+
+#define LOCTEXT_NAMESPACE "UTFEditorModule"
 
 /*******************************************************************************
  * 
@@ -14,6 +21,18 @@ namespace UnitTestFrameworkEditorModuleImpl
 {
 	/**  */
 	static void AddUnitTestingRegistryTags(const UWorld* World, TArray<UObject::FAssetRegistryTag>& TagsOut);
+
+	/**  */
+	static void ExtendBlueprintEditorMenus(TSharedPtr<FExtender> MenuExtender, UBlueprint* Blueprint);
+
+	/**  */
+	static FBlueprintEditorModule& GetBlueprintEditorModule();
+
+	/**  */
+	static void ExtendBlueprintToolbar(FToolBarBuilder& MenuBuilder, UBlueprint* Blueprint);
+	
+	/**  */
+	TSharedRef<SWidget> GenerateBlueprintMenuContent(UBlueprint* Blueprint);
 }
 
 //------------------------------------------------------------------------------
@@ -29,15 +48,97 @@ static void UnitTestFrameworkEditorModuleImpl::AddUnitTestingRegistryTags(const
 			}
 
 			UClass* ActorClass = Actor->GetClass();
-			if (ActorClass->ImplementsInterface(UUTFUnitTestInterface::StaticClass()))
+			if (FUTFEditorAutomationTests::IsClassSuitableForInEditorTesting(ActorClass))
 			{
-				TagsOut.Add(UObject::FAssetRegistryTag(FUTFEditorAutomationTests::UnitTestAssetTag, TEXT("TRUE"), UObject::FAssetRegistryTag::TT_Hidden));
-				break;
+				TagsOut.Add(UObject::FAssetRegistryTag(FUTFEditorAutomationTests::UnitTestLevelTag, TEXT("TRUE"), UObject::FAssetRegistryTag::TT_Hidden));
 			}
 		}
 	}
 }
 
+//------------------------------------------------------------------------------
+static void UnitTestFrameworkEditorModuleImpl::ExtendBlueprintEditorMenus(TSharedPtr<FExtender> MenuExtender, UBlueprint* Blueprint)
+{
+	FToolBarExtensionDelegate ToolbarExtension = FToolBarExtensionDelegate::CreateStatic(ExtendBlueprintToolbar, Blueprint);
+
+	FBlueprintEditorModule& BPEditor = UnitTestFrameworkEditorModuleImpl::GetBlueprintEditorModule();
+	MenuExtender->AddToolBarExtension("Asset", EExtensionHook::After, BPEditor.GetsSharedBlueprintEditorCommands(), ToolbarExtension);
+}
+
+//------------------------------------------------------------------------------
+static FBlueprintEditorModule& UnitTestFrameworkEditorModuleImpl::GetBlueprintEditorModule()
+{
+	return FModuleManager::LoadModuleChecked<FBlueprintEditorModule>("Kismet");
+}
+
+//------------------------------------------------------------------------------
+static void UnitTestFrameworkEditorModuleImpl::ExtendBlueprintToolbar(FToolBarBuilder& MenuBuilder, UBlueprint* Blueprint)
+{
+	auto HasUTFInterface = [Blueprint]()->bool
+	{
+		UClass* BlueprintClass = Blueprint->SkeletonGeneratedClass;
+		return BlueprintClass && BlueprintClass->ImplementsInterface(UUTFUnitTestInterface::StaticClass());
+	};
+
+	//MenuBuilder.BeginSection(TEXT("UnitTestFramework"));
+	{
+		MenuBuilder.AddComboButton(
+			FUIAction(
+				FExecuteAction(),
+				FCanExecuteAction::CreateLambda(HasUTFInterface),
+				FIsActionChecked(),
+				FIsActionButtonVisible::CreateLambda(HasUTFInterface)
+			),
+			FOnGetContent::CreateStatic(GenerateBlueprintMenuContent, Blueprint),
+			LOCTEXT("BlueprintOptionsButton", "Unit Test"),
+			LOCTEXT("BlueprintOptionsButton_Tooltip", "Options to customize how this Blueprint is used with the UnitTestFramework."),
+			FSlateIcon("UTFEditorStyle", "UTFBlueprint.SettingsIcon", "UTFBlueprint.SettingsIcon"),
+			/*bSimpleComboBox =*/false
+		);
+	}
+	//MenuBuilder.EndSection();
+}
+
+//------------------------------------------------------------------------------
+TSharedRef<SWidget> UnitTestFrameworkEditorModuleImpl::GenerateBlueprintMenuContent(UBlueprint* Blueprint)
+{
+	FBlueprintEditorModule& BPEditor = UnitTestFrameworkEditorModuleImpl::GetBlueprintEditorModule();
+	FMenuBuilder MenuBuilder(/*bShouldCloseWindowAfterMenuSelection =*/false, BPEditor.GetsSharedBlueprintEditorCommands());
+
+	auto ToggleEditorSafeTest = [Blueprint]()
+	{
+		UPackage* BlueprintPackage = Blueprint->GetOutermost();
+		UMetaData* MetaData = BlueprintPackage->GetMetaData();
+
+		if (MetaData->HasValue(Blueprint, FUTFEditorAutomationTests::InEditorTestTag))
+		{
+			MetaData->RemoveValue(Blueprint, FUTFEditorAutomationTests::InEditorTestTag);
+		}
+		else
+		{
+			MetaData->SetValue(Blueprint, FUTFEditorAutomationTests::InEditorTestTag, TEXT("TRUE"));
+		}
+
+		BlueprintPackage->MarkPackageDirty();
+	};
+	UClass* BlueprintClass = Blueprint->SkeletonGeneratedClass;
+
+	MenuBuilder.AddMenuEntry(
+		LOCTEXT("EditorSafeTest_Label", "In-Editor Test"),
+		LOCTEXT("EditorSafeTest_Tooltip", "Flags this test as safe to run in the editor; meaning its script does not execute any actions that are prevented in the construction script (like SpawnActor, latent actions, etc.))"),
+		FSlateIcon(),
+		FUIAction(
+			FExecuteAction::CreateLambda(ToggleEditorSafeTest),
+			FCanExecuteAction(),
+			FIsActionChecked::CreateStatic(FUTFEditorAutomationTests::IsClassSuitableForInEditorTesting, BlueprintClass)
+		),
+		NAME_None,
+		EUserInterfaceActionType::ToggleButton
+	);
+
+	return MenuBuilder.MakeWidget();
+}
+
 /*******************************************************************************
  * FUnitTestFrameworkEditorModule
  ******************************************************************************/
@@ -48,18 +149,28 @@ static void UnitTestFrameworkEditorModuleImpl::AddUnitTestingRegistryTags(const
 class FUnitTestFrameworkEditorModule : public IUnitTestFrameworkEditorModuleInterface
 {
 public:
+	//--------------------------------------------------------------------------
 	virtual void StartupModule() override
 	{
+		FUTFEditorStyle::Initialize();
 		GetWorldAssetTagsHandle = FWorldDelegates::GetAssetTags.AddStatic(UnitTestFrameworkEditorModuleImpl::AddUnitTestingRegistryTags);
+
+		FBlueprintEditorModule& BPEditor = UnitTestFrameworkEditorModuleImpl::GetBlueprintEditorModule();
+		BPEditor.OnGatherBlueprintMenuExtensions().AddStatic(UnitTestFrameworkEditorModuleImpl::ExtendBlueprintEditorMenus);
 	}
 
+	//--------------------------------------------------------------------------
 	virtual void ShutdownModule() override
 	{
 		FWorldDelegates::GetAssetTags.Remove(GetWorldAssetTagsHandle);
 		GetWorldAssetTagsHandle.Reset();
+
+		FUTFEditorStyle::Shutdown();
 	}
 
 private: 
 	FDelegateHandle GetWorldAssetTagsHandle;
 };
+
 IMPLEMENT_MODULE(FUnitTestFrameworkEditorModule, UnitTestFrameworkEditor);
+#undef LOCTEXT_NAMESPACE
diff --git a/Engine/Plugins/Experimental/UnitTestFramework/Source/UnitTestFrameworkEditor/UnitTestFrameworkEditor.Build.cs b/Engine/Plugins/Experimental/UnitTestFramework/Source/UnitTestFrameworkEditor/UnitTestFrameworkEditor.Build.cs
--- a/Engine/Plugins/Experimental/UnitTestFramework/Source/UnitTestFrameworkEditor/UnitTestFrameworkEditor.Build.cs
+++ b/Engine/Plugins/Experimental/UnitTestFramework/Source/UnitTestFrameworkEditor/UnitTestFrameworkEditor.Build.cs
@@ -21,6 +21,11 @@ public class UnitTestFrameworkEditor : ModuleRules
                 "UnitTestFramework",
                 "UnrealEd",
                 "AssetRegistry",
+                "SlateCore",
+                "Kismet",
+                "Slate",
+                "EditorStyle",
+                "Projects",
 			}
 		);
 	}
diff --git a/Engine/Source/Editor/Kismet/Private/BlueprintEditor.cpp b/Engine/Source/Editor/Kismet/Private/BlueprintEditor.cpp
--- a/Engine/Source/Editor/Kismet/Private/BlueprintEditor.cpp
+++ b/Engine/Source/Editor/Kismet/Private/BlueprintEditor.cpp
@@ -1637,6 +1637,7 @@ void FBlueprintEditor::InitBlueprintEditor(
 
 	FBlueprintEditorModule* BlueprintEditorModule = &FModuleManager::LoadModuleChecked<FBlueprintEditorModule>("Kismet");
 	TSharedPtr<FExtender> CustomExtenders = BlueprintEditorModule->GetMenuExtensibilityManager()->GetAllExtenders(GetToolkitCommands(), GetEditingObjects());
+	BlueprintEditorModule->OnGatherBlueprintMenuExtensions().Broadcast(CustomExtenders, GetBlueprintObj());
 
 	AddMenuExtender(CustomExtenders);
 	AddToolbarExtender(CustomExtenders);
diff --git a/Engine/Source/Editor/Kismet/Public/BlueprintEditorModule.h b/Engine/Source/Editor/Kismet/Public/BlueprintEditorModule.h
--- a/Engine/Source/Editor/Kismet/Public/BlueprintEditorModule.h
+++ b/Engine/Source/Editor/Kismet/Public/BlueprintEditorModule.h
@@ -124,6 +124,10 @@ public:
 	/** Gets the extensibility managers for outside entities to extend blueprint editor's menus and toolbars */
 	virtual TSharedPtr<FExtensibilityManager> GetMenuExtensibilityManager() override { return MenuExtensibilityManager; }
 
+	/**  */
+	DECLARE_EVENT_TwoParams(FBlueprintEditorModule, FBlueprintMenuExtensionEvent, TSharedPtr<FExtender>, UBlueprint*);
+	FBlueprintMenuExtensionEvent& OnGatherBlueprintMenuExtensions() { return GatherBlueprintMenuExtensions; }
+
 	/** 
 	 * Register a customization for interacting with the SCS editor 
 	 * @param	InComponentName			The name of the component to customize behavior for
@@ -155,6 +159,9 @@ private:
 private:
 	TSharedPtr<FExtensibilityManager> MenuExtensibilityManager;
 
+	//
+	FBlueprintMenuExtensionEvent GatherBlueprintMenuExtensions;
+
 	// Event to be called when the blueprint editor is opened
 	FBlueprintEditorOpenedEvent BlueprintEditorOpened;
 
