CIS fix. Missing files.

diff --git a/Engine/Plugins/Experimental/UnitTestFramework/Source/UnitTestFrameworkEditor/Private/UTFEditorStyle.cpp b/Engine/Plugins/Experimental/UnitTestFramework/Source/UnitTestFrameworkEditor/Private/UTFEditorStyle.cpp
new file mode 100644
--- /dev/null
+++ b/Engine/Plugins/Experimental/UnitTestFramework/Source/UnitTestFrameworkEditor/Private/UTFEditorStyle.cpp
@@ -0,0 +1,93 @@
+// Copyright 1998-2015 Epic Games, Inc. All Rights Reserved.
+
+#include "UnitTestFrameworkEditorPCH.h"
+#include "UTFEditorStyle.h"
+#include "SlateStyle.h"
+#include "IPluginManager.h"
+
+#define IMAGE_PLUGIN_BRUSH( RelativePath, ... ) FSlateImageBrush(UTFEditorStyleImpl::InContent(RelativePath, ".png"), __VA_ARGS__)
+#define IMAGE_BRUSH(RelativePath, ...)			FSlateImageBrush(StyleSet->RootToContentDir(RelativePath, TEXT(".png")), __VA_ARGS__)
+#define BOX_BRUSH(RelativePath, ...)			FSlateBoxBrush(StyleSet->RootToContentDir(RelativePath, TEXT(".png")), __VA_ARGS__)
+#define TTF_FONT(RelativePath, ...)				FSlateFontInfo(StyleSet->RootToContentDir(RelativePath, TEXT(".ttf")), __VA_ARGS__)
+#define TTF_CORE_FONT(RelativePath, ...)		FSlateFontInfo(StyleSet->RootToCoreContentDir(RelativePath, TEXT(".ttf") ), __VA_ARGS__)
+
+/*******************************************************************************
+ * UTFEditorStyleImpl
+ ******************************************************************************/
+ 
+namespace UTFEditorStyleImpl
+{
+	static TSharedPtr<FSlateStyleSet> StyleSet = nullptr;
+	/**  */
+	static FString InContent(const FString& RelativePath, const ANSICHAR* Extension);
+}
+
+//------------------------------------------------------------------------------
+FString UTFEditorStyleImpl::InContent(const FString& RelativePath, const ANSICHAR* Extension)
+{
+	static FString ContentDir = IPluginManager::Get().FindPlugin(TEXT("UnitTestFramework"))->GetContentDir();
+	return (ContentDir / RelativePath) + Extension;
+}
+
+/*******************************************************************************
+ * FUTFEditorStyle
+ ******************************************************************************/
+
+//------------------------------------------------------------------------------
+TSharedPtr<ISlateStyle> FUTFEditorStyle::Get()
+{ 
+	return UTFEditorStyleImpl::StyleSet;
+}
+
+//------------------------------------------------------------------------------
+FName FUTFEditorStyle::GetStyleSetName()
+{
+	static FName UTFStyleName(TEXT("UTFEditorStyle"));
+	return UTFStyleName;
+}
+
+//------------------------------------------------------------------------------
+void FUTFEditorStyle::Initialize()
+{
+	TSharedPtr<FSlateStyleSet>& StyleSet = UTFEditorStyleImpl::StyleSet;
+	// Only register once
+	if (StyleSet.IsValid())
+	{
+		return;
+	}
+
+	const FVector2D Icon8x8(8.0f, 8.0f);
+	const FVector2D Icon16x16(16.0f, 16.0f);
+	const FVector2D Icon20x20(20.0f, 20.0f);
+	const FVector2D Icon40x40(40.0f, 40.0f);
+	const FTextBlockStyle& NormalText = FEditorStyle::Get().GetWidgetStyle<FTextBlockStyle>("NormalText");
+
+	StyleSet = MakeShareable(new FSlateStyleSet(GetStyleSetName()));
+	StyleSet->SetContentRoot(FPaths::EngineContentDir() / TEXT("Editor/Slate"));
+	StyleSet->SetCoreContentRoot(FPaths::EngineContentDir() / TEXT("Slate"));
+
+	// Blueprint editor
+	{
+		StyleSet->Set("UTFBlueprint.SettingsIcon", new IMAGE_PLUGIN_BRUSH(TEXT("Editor/UnitTestFrameworkIcon"), Icon40x40));
+	}
+
+	FSlateStyleRegistry::RegisterSlateStyle(*StyleSet.Get());
+};
+
+#undef IMAGE_PLUGIN_BRUSH
+#undef IMAGE_BRUSH
+#undef BOX_BRUSH
+#undef TTF_FONT
+#undef TTF_CORE_FONT
+
+//------------------------------------------------------------------------------
+void FUTFEditorStyle::Shutdown()
+{
+	TSharedPtr<FSlateStyleSet>& StyleSet = UTFEditorStyleImpl::StyleSet;
+	if (StyleSet.IsValid())
+	{
+		FSlateStyleRegistry::UnRegisterSlateStyle(*StyleSet.Get());
+		ensure(StyleSet.IsUnique());
+		StyleSet.Reset();
+	}
+}
diff --git a/Engine/Plugins/Experimental/UnitTestFramework/Source/UnitTestFrameworkEditor/Private/UTFEditorStyle.h b/Engine/Plugins/Experimental/UnitTestFramework/Source/UnitTestFrameworkEditor/Private/UTFEditorStyle.h
new file mode 100644
--- /dev/null
+++ b/Engine/Plugins/Experimental/UnitTestFramework/Source/UnitTestFrameworkEditor/Private/UTFEditorStyle.h
@@ -0,0 +1,18 @@
+// Copyright 1998-2015 Epic Games, Inc. All Rights Reserved.
+
+#pragma once
+
+class ISlateStyle;
+
+/** 
+ * 
+ */
+class FUTFEditorStyle
+{
+public:
+	static void Initialize();
+	static void Shutdown();
+
+	static TSharedPtr<ISlateStyle> Get();
+	static FName GetStyleSetName();
+};
\ No newline at end of file
