CIS non-unity compile fix: Missing include.

diff --git a/Engine/Plugins/Experimental/UnitTestFramework/Source/UnitTestFrameworkEditor/Private/UTFEditorStyle.cpp b/Engine/Plugins/Experimental/UnitTestFramework/Source/UnitTestFrameworkEditor/Private/UTFEditorStyle.cpp
--- a/Engine/Plugins/Experimental/UnitTestFramework/Source/UnitTestFrameworkEditor/Private/UTFEditorStyle.cpp
+++ b/Engine/Plugins/Experimental/UnitTestFramework/Source/UnitTestFrameworkEditor/Private/UTFEditorStyle.cpp
@@ -4,6 +4,7 @@
 #include "UTFEditorStyle.h"
 #include "SlateStyle.h"
 #include "IPluginManager.h"
+#include "EditorStyleSet.h"
 
 #define IMAGE_PLUGIN_BRUSH( RelativePath, ... ) FSlateImageBrush(UTFEditorStyleImpl::InContent(RelativePath, ".png"), __VA_ARGS__)
 #define IMAGE_BRUSH(RelativePath, ...)			FSlateImageBrush(StyleSet->RootToContentDir(RelativePath, TEXT(".png")), __VA_ARGS__)
@@ -60,7 +61,7 @@ void FUTFEditorStyle::Initialize()
 	const FVector2D Icon16x16(16.0f, 16.0f);
 	const FVector2D Icon20x20(20.0f, 20.0f);
 	const FVector2D Icon40x40(40.0f, 40.0f);
-	const FTextBlockStyle& NormalText = FEditorStyle::Get().GetWidgetStyle<FTextBlockStyle>("NormalText");
+	//const FTextBlockStyle& NormalText = FEditorStyle::Get().GetWidgetStyle<FTextBlockStyle>("NormalText");
 
 	StyleSet = MakeShareable(new FSlateStyleSet(GetStyleSetName()));
 	StyleSet->SetContentRoot(FPaths::EngineContentDir() / TEXT("Editor/Slate"));
